Fixed MIDIDevice module public header dependency issue (UE-40704)
- This fixes compilation errors when including MIDIDeviceController.h from an external module

diff --git a/Engine/Plugins/Runtime/MIDIDevice/Source/MIDIDevice/MIDIDeviceController.cpp b/Engine/Plugins/Runtime/MIDIDevice/Source/MIDIDevice/MIDIDeviceController.cpp
--- a/Engine/Plugins/Runtime/MIDIDevice/Source/MIDIDevice/MIDIDeviceController.cpp
+++ b/Engine/Plugins/Runtime/MIDIDevice/Source/MIDIDevice/MIDIDeviceController.cpp
@@ -2,6 +2,7 @@
 
 #include "MIDIDeviceController.h"
 #include "MIDIDeviceLog.h"
+#include "portmidi.h"
 
 
 #define LOCTEXT_NAMESPACE "MIDIDeviceController"
diff --git a/Engine/Plugins/Runtime/MIDIDevice/Source/MIDIDevice/MIDIDeviceManager.cpp b/Engine/Plugins/Runtime/MIDIDevice/Source/MIDIDevice/MIDIDeviceManager.cpp
--- a/Engine/Plugins/Runtime/MIDIDevice/Source/MIDIDevice/MIDIDeviceManager.cpp
+++ b/Engine/Plugins/Runtime/MIDIDevice/Source/MIDIDevice/MIDIDeviceManager.cpp
@@ -5,6 +5,7 @@
 #include "UObject/UObjectHash.h"
 #include "MIDIDeviceController.h"
 #include "MIDIDeviceLog.h"
+#include "portmidi.h"
 
 #define LOCTEXT_NAMESPACE "MIDIDeviceManager"
 
diff --git a/Engine/Plugins/Runtime/MIDIDevice/Source/MIDIDevice/Public/MIDIDeviceController.h b/Engine/Plugins/Runtime/MIDIDevice/Source/MIDIDevice/Public/MIDIDeviceController.h
--- a/Engine/Plugins/Runtime/MIDIDevice/Source/MIDIDevice/Public/MIDIDeviceController.h
+++ b/Engine/Plugins/Runtime/MIDIDevice/Source/MIDIDevice/Public/MIDIDeviceController.h
@@ -6,7 +6,6 @@
 #include "UObject/ObjectMacros.h"
 #include "UObject/Object.h"
 #include "UObject/WeakObjectPtr.h"
-#include "portmidi.h"
 #include "MIDIDeviceController.generated.h"
 
 
@@ -82,7 +81,7 @@ protected:
 	// ...
 
 	/** The PortMidi stream used for MIDI input for this device */
-	PortMidiStream* PMMIDIInputStream;
+	void* PMMIDIInputStream;
 
 	/** Size of the MIDI buffer in bytes */
 	int32 MIDIBufferSize;
