From Maxwell.Hayes on Dev-Audio:
Fix for Pitch Bend Midi-parsing bug.
(MSB/LSB were flipped.  Changed bit-zeroing via cast to bitwise &)

diff --git a/Engine/Plugins/Runtime/MIDIDevice/Source/MIDIDevice/MIDIDeviceInputController.cpp b/Engine/Plugins/Runtime/MIDIDevice/Source/MIDIDevice/MIDIDeviceInputController.cpp
--- a/Engine/Plugins/Runtime/MIDIDevice/Source/MIDIDevice/MIDIDeviceInputController.cpp
+++ b/Engine/Plugins/Runtime/MIDIDevice/Source/MIDIDevice/MIDIDeviceInputController.cpp
@@ -145,7 +145,7 @@ void UMIDIDeviceInputController::ProcessIncomingMIDIEvents()
 					{
 						int32 pitchBend = (PMMessageData1 & 0x7F) << 7;
 
-						pitchBend |= static_cast<unsigned short>(PMMessageData2);
+						pitchBend |= (PMMessageData1 & 0x7F);
 
 						this->OnMIDIPitchBend.Broadcast(this, Timestamp, Channel, pitchBend);
 
