Fix old copyright issues #jira https://jira.it.epicgames.net/browse/UE-69064 #rb none

diff --git a/Engine/Plugins/Runtime/MIDIDevice/Source/MIDIDevice/MIDIDeviceInputController.cpp b/Engine/Plugins/Runtime/MIDIDevice/Source/MIDIDevice/MIDIDeviceInputController.cpp
--- a/Engine/Plugins/Runtime/MIDIDevice/Source/MIDIDevice/MIDIDeviceInputController.cpp
+++ b/Engine/Plugins/Runtime/MIDIDevice/Source/MIDIDevice/MIDIDeviceInputController.cpp
@@ -1,4 +1,4 @@
-// Copyright 1998-2018 Epic Games, Inc. All Rights Reserved.
+// Copyright 1998-2019 Epic Games, Inc. All Rights Reserved.
 
 #include "MIDIDeviceInputController.h"
 #include "MIDIDeviceLog.h"
diff --git a/Engine/Plugins/Runtime/MIDIDevice/Source/MIDIDevice/MIDIDeviceOutputController.cpp b/Engine/Plugins/Runtime/MIDIDevice/Source/MIDIDevice/MIDIDeviceOutputController.cpp
--- a/Engine/Plugins/Runtime/MIDIDevice/Source/MIDIDevice/MIDIDeviceOutputController.cpp
+++ b/Engine/Plugins/Runtime/MIDIDevice/Source/MIDIDevice/MIDIDeviceOutputController.cpp
@@ -1,4 +1,4 @@
-// Copyright 1998-2018 Epic Games, Inc. All Rights Reserved.
+// Copyright 1998-2019 Epic Games, Inc. All Rights Reserved.
 
 #include "MIDIDeviceOutputController.h"
 #include "MIDIDeviceLog.h"
diff --git a/Engine/Plugins/Runtime/MIDIDevice/Source/MIDIDevice/Public/MIDIDeviceInputController.h b/Engine/Plugins/Runtime/MIDIDevice/Source/MIDIDevice/Public/MIDIDeviceInputController.h
--- a/Engine/Plugins/Runtime/MIDIDevice/Source/MIDIDevice/Public/MIDIDeviceInputController.h
+++ b/Engine/Plugins/Runtime/MIDIDevice/Source/MIDIDevice/Public/MIDIDeviceInputController.h
@@ -1,4 +1,4 @@
-// Copyright 1998-2018 Epic Games, Inc. All Rights Reserved.
+// Copyright 1998-2019 Epic Games, Inc. All Rights Reserved.
 
 #pragma once
 
diff --git a/Engine/Plugins/Runtime/MIDIDevice/Source/MIDIDevice/Public/MIDIDeviceOutputController.h b/Engine/Plugins/Runtime/MIDIDevice/Source/MIDIDevice/Public/MIDIDeviceOutputController.h
--- a/Engine/Plugins/Runtime/MIDIDevice/Source/MIDIDevice/Public/MIDIDeviceOutputController.h
+++ b/Engine/Plugins/Runtime/MIDIDevice/Source/MIDIDevice/Public/MIDIDeviceOutputController.h
@@ -1,4 +1,4 @@
-// Copyright 1998-2018 Epic Games, Inc. All Rights Reserved.
+// Copyright 1998-2019 Epic Games, Inc. All Rights Reserved.
 
 #pragma once
 
