Uninitialized variable warning fixes

diff --git a/Engine/Plugins/FX/CascadeToNiagaraConverter/Source/CascadeToNiagaraConverter/Public/NiagaraStackGraphUtilitiesAdapterLibrary.h b/Engine/Plugins/FX/CascadeToNiagaraConverter/Source/CascadeToNiagaraConverter/Public/NiagaraStackGraphUtilitiesAdapterLibrary.h
--- a/Engine/Plugins/FX/CascadeToNiagaraConverter/Source/CascadeToNiagaraConverter/Public/NiagaraStackGraphUtilitiesAdapterLibrary.h
+++ b/Engine/Plugins/FX/CascadeToNiagaraConverter/Source/CascadeToNiagaraConverter/Public/NiagaraStackGraphUtilitiesAdapterLibrary.h
@@ -433,7 +433,7 @@ struct FStackEntryAddAction
 
 	// The mode of the AddStackEntryAction: when the action is applied to a UNiagaraEmitterConversionContext during Finalize(), the mode chooses what form of stack entry to add.
 	UPROPERTY()
-	EStackEntryAddActionMode Mode;
+	EStackEntryAddActionMode Mode = EStackEntryAddActionMode::Module;
 
 	// If Mode is Module, represents the pending module script to create a stack entry for. Otherwise this value is ignored.
 	UPROPERTY()
diff --git a/Engine/Plugins/Online/OnlineFramework/Source/Qos/Public/QosRegionManager.h b/Engine/Plugins/Online/OnlineFramework/Source/Qos/Public/QosRegionManager.h
--- a/Engine/Plugins/Online/OnlineFramework/Source/Qos/Public/QosRegionManager.h
+++ b/Engine/Plugins/Online/OnlineFramework/Source/Qos/Public/QosRegionManager.h
@@ -52,7 +52,7 @@ struct FQosPingServerInfo
 	FString Address;
 	/** Port of server */
 	UPROPERTY()
-	int32 Port;
+	int32 Port = 0;
 };
 
 /**
diff --git a/Engine/Plugins/Runtime/MIDIDevice/Source/MIDIDevice/Public/MIDIDeviceManager.h b/Engine/Plugins/Runtime/MIDIDevice/Source/MIDIDevice/Public/MIDIDeviceManager.h
--- a/Engine/Plugins/Runtime/MIDIDevice/Source/MIDIDevice/Public/MIDIDeviceManager.h
+++ b/Engine/Plugins/Runtime/MIDIDevice/Source/MIDIDevice/Public/MIDIDeviceManager.h
@@ -16,7 +16,7 @@ struct MIDIDEVICE_API FFoundMIDIDevice
 
 	/** The unique ID of this MIDI device */
 	UPROPERTY(BlueprintReadOnly, Category="MIDI Device Manager")
-	int32 DeviceID;
+	int32 DeviceID = 0;
 
 	/** The name of this device.  This name comes from the MIDI hardware, and might not be unique */
 	UPROPERTY(BlueprintReadOnly, Category="MIDI Device Manager")
@@ -24,23 +24,23 @@ struct MIDIDEVICE_API FFoundMIDIDevice
 
 	/** True if the device supports sending events to us */
 	UPROPERTY(BlueprintReadOnly, Category="MIDI Device Manager")
-	bool bCanReceiveFrom;
+	bool bCanReceiveFrom = false;
 
 	/** True if the device supports receiving events from us */
 	UPROPERTY(BlueprintReadOnly, Category="MIDI Device Manager")
-	bool bCanSendTo;
+	bool bCanSendTo = false;
 
 	/** Whether the device is already in use.  You might not want to create a controller for devices that are busy.  Someone else could be using it. */
 	UPROPERTY(BlueprintReadOnly, Category="MIDI Device Manager")
-	bool bIsAlreadyInUse;
+	bool bIsAlreadyInUse = false;
 
 	/** True if this is the default MIDI device for input on this system */
 	UPROPERTY(BlueprintReadOnly, Category="MIDI Device Manager")
-	bool bIsDefaultInputDevice;
+	bool bIsDefaultInputDevice = false;
 
 	/** True if this is the default MIDI device for output on this system */
 	UPROPERTY(BlueprintReadOnly, Category="MIDI Device Manager")
-	bool bIsDefaultOutputDevice;
+	bool bIsDefaultOutputDevice = false;
 };
 
 
@@ -51,7 +51,7 @@ struct MIDIDEVICE_API FMIDIDeviceInfo
 
 	/** The unique ID of this MIDI device */
 	UPROPERTY(BlueprintReadOnly, Category = "MIDI Device Manager")
-	int32 DeviceID;
+	int32 DeviceID = 0;
 
 	/** The name of this device.  This name comes from the MIDI hardware, any might not be unique */
 	UPROPERTY(BlueprintReadOnly, Category = "MIDI Device Manager")
@@ -59,11 +59,11 @@ struct MIDIDEVICE_API FMIDIDeviceInfo
 
 	/** Whether the device is already in use.  You might not want to create a controller for devices that are busy.  Someone else could be using it. */
 	UPROPERTY(BlueprintReadOnly, Category = "MIDI Device Manager")
-	bool bIsAlreadyInUse;
+	bool bIsAlreadyInUse = false;
 
 	/** True if this is the default MIDI device for input on this system */
 	UPROPERTY(BlueprintReadOnly, Category = "MIDI Device Manager")
-	bool bIsDefaultDevice;
+	bool bIsDefaultDevice = false;
 };
 
 UCLASS()
diff --git a/Engine/Plugins/Runtime/NetworkPrediction/Source/NetworkPrediction/Public/NetworkPredictionAsync.h b/Engine/Plugins/Runtime/NetworkPrediction/Source/NetworkPrediction/Public/NetworkPredictionAsync.h
--- a/Engine/Plugins/Runtime/NetworkPrediction/Source/NetworkPrediction/Public/NetworkPredictionAsync.h
+++ b/Engine/Plugins/Runtime/NetworkPrediction/Source/NetworkPrediction/Public/NetworkPredictionAsync.h
@@ -1024,6 +1024,11 @@ struct FTempMockInputCmd
 	UPROPERTY(BlueprintReadWrite,Category="Input")
 	bool bBrakesPressed = false;
 
+	FTempMockInputCmd()
+		: Force(ForceInitToZero)
+	{
+	}
+
 	void NetSerialize(FArchive& Ar)
 	{
 		Ar << Force;
