PR #9656: [bugfix] Fix FindAllMIDIDeviceInfo having no effect. 'Get midi device? (Contributed by [REDACTED_PERSON])

diff --git a/Engine/Plugins/Runtime/MIDIDevice/Source/MIDIDevice/Private/MIDIDeviceManager.cpp b/Engine/Plugins/Runtime/MIDIDevice/Source/MIDIDevice/Private/MIDIDeviceManager.cpp
--- a/Engine/Plugins/Runtime/MIDIDevice/Source/MIDIDevice/Private/MIDIDeviceManager.cpp
+++ b/Engine/Plugins/Runtime/MIDIDevice/Source/MIDIDevice/Private/MIDIDeviceManager.cpp
@@ -226,6 +226,7 @@ void UMIDIDeviceManager::FindAllMIDIDeviceInfo(TArray<FMIDIDeviceInfo>& OutMIDII
 	{
 		ReinitializeDeviceManager();
 	}
+	FindAllMIDIDeviceInfoInternal(OutMIDIInputDevicesInfo, OutMIDIOutputDevicesInfo);
 }
 
 void UMIDIDeviceManager::GetMIDIInputDeviceIDByName(const FString DeviceName, int32& DeviceID)
