Add static to Midi Device Manager to force-cleanup any dangling controller objects

#rnx

diff --git a/Engine/Plugins/Runtime/MIDIDevice/Source/MIDIDevice/Private/MIDIDeviceManager.cpp b/Engine/Plugins/Runtime/MIDIDevice/Source/MIDIDevice/Private/MIDIDeviceManager.cpp
--- a/Engine/Plugins/Runtime/MIDIDevice/Source/MIDIDevice/Private/MIDIDeviceManager.cpp
+++ b/Engine/Plugins/Runtime/MIDIDevice/Source/MIDIDevice/Private/MIDIDeviceManager.cpp
@@ -154,6 +154,17 @@ void UMIDIDeviceManager::FindMIDIDevicesInternal(TArray<FFoundMIDIDevice>& OutMI
 	}
 }
 
+void UMIDIDeviceManager::ShutDownAllMIDIDevices()
+{
+	for (TObjectIterator<UMIDIDeviceControllerBase> ControllerIter; ControllerIter; ++ControllerIter)
+	{
+		UMIDIDeviceControllerBase* MIDIDeviceController = *ControllerIter;
+		if (IsValid(MIDIDeviceController))
+		{
+			MIDIDeviceController->ShutdownDevice();
+		}
+	}
+}
 
 void UMIDIDeviceManager::FindMIDIDevices(TArray<FFoundMIDIDevice>& OutMIDIDevices)
 {
diff --git a/Engine/Plugins/Runtime/MIDIDevice/Source/MIDIDevice/Public/MIDIDeviceManager.h b/Engine/Plugins/Runtime/MIDIDevice/Source/MIDIDevice/Public/MIDIDeviceManager.h
--- a/Engine/Plugins/Runtime/MIDIDevice/Source/MIDIDevice/Public/MIDIDeviceManager.h
+++ b/Engine/Plugins/Runtime/MIDIDevice/Source/MIDIDevice/Public/MIDIDeviceManager.h
@@ -73,6 +73,11 @@ class MIDIDEVICE_API UMIDIDeviceManager : public UBlueprintFunctionLibrary
 
 public:
 
+	/**
+	 * Shuts down any existing midi devices to ensure no dangling streams are left, and mark any input or output controller for cleanup
+	 */
+	UFUNCTION(BlueprintCallable, Category="MIDI Device Manager")
+	static void ShutDownAllMIDIDevices();
 	/**
 	 * Enumerates all of the connected MIDI devices and reports back with the IDs and names of those devices.  This operation is a little expensive
 	 * so only do it once at startup, or if you think that a new device may have been connected.
